TP_MFG_ProConfigBP blueprint and content adjustments for transfer to UE5.0

diff --git a/Templates/TP_MFG_ProdConfigBP/Config/DefaultEngine.ini b/Templates/TP_MFG_ProdConfigBP/Config/DefaultEngine.ini
--- a/Templates/TP_MFG_ProdConfigBP/Config/DefaultEngine.ini
+++ b/Templates/TP_MFG_ProdConfigBP/Config/DefaultEngine.ini
@@ -12,6 +12,8 @@ AppliedDefaultGraphicsPerformance=Maximum
 
 [/Script/Engine.RendererSettings]
 r.DefaultFeature.AutoExposure.ExtendDefaultLuminanceRange=True
+r.DefaultFeature.AutoExposure=False
+r.GenerateMeshDistanceFields=True
 
 [/Script/EngineSettings.GameMapsSettings]
 EditorStartupMap=/Game/ProductAssets/Levels/Main.Main
