Updating ja translation for a few templates and feature packs

diff --git a/Templates/TP_AEC_ArchvisBP/Config/TemplateDefs.ini b/Templates/TP_AEC_ArchvisBP/Config/TemplateDefs.ini
--- a/Templates/TP_AEC_ArchvisBP/Config/TemplateDefs.ini
+++ b/Templates/TP_AEC_ArchvisBP/Config/TemplateDefs.ini
@@ -19,6 +19,7 @@ LocalizedDescriptions=(Language="ko",Text="이 템플릿은 건축 시각화 워
 LocalizedDisplayNames=(Language="zh-Hans",Text="建筑可视化")
 LocalizedDescriptions=(Language="zh-Hans",Text="该模板专为建筑可视化流程设计，包含数个室外及室内静态图片渲染、动画、阳光研究以及非写实风格化渲染方面的示例。\n\n该模板启用了Datasmith和其他实用插件，如影片渲染队列以及物理精准的太阳和天空系统，可以让你快速启动建筑可视化项目。")
 
+LocalizedDisplayNames=(Language="ja", Text="建築ビジュアライゼーション")
 LocalizedDescriptions=(Language="ja",Text="このテンプレートは建築ビジュアライゼーション向けです。エクステリア及びインテリアの静止画レンダリング、アニメーション、日照検討、ノンフォトリアル レンダリングの例を含んでいます。\n\n建築ビジュアライゼーション プロジェクトをすぐに開始できるように準備したテンプレートです。Datasmith に加えて ムービー レンダー キューなど便利なプラグインを有効にし、物理的に正確な太陽、空システムも有効にしています。")
 
 
diff --git a/Templates/TP_AEC_BlankBP/Config/TemplateDefs.ini b/Templates/TP_AEC_BlankBP/Config/TemplateDefs.ini
--- a/Templates/TP_AEC_BlankBP/Config/TemplateDefs.ini
+++ b/Templates/TP_AEC_BlankBP/Config/TemplateDefs.ini
@@ -21,7 +21,8 @@ LocalizedDescriptions=(Language="ko", Text="최소한의 기능과 콘텐츠를
 LocalizedDisplayNames=(Language="zh-Hans", Text="空白")
 LocalizedDescriptions=(Language="zh-Hans", Text="具备基本功能和内容的空白模板。\n\n该模板允许使用Datasmith从多种设计软件中导入模型并轻松为其设置光照。")
 
-LocalizedDescriptions=(Language="ja", Text="最小限の機能とコンテンツが備わった Blank テンプレートです。\n\nこのテンプレートを利用すると、Datasmith を使ってモデルをさまざまなデザイン アプリケーションからインポートできるようになるとともに、それらを簡単にライティングできるようになります。")
+LocalizedDisplayNames=(Language="ja", Text="ブランク")
+LocalizedDescriptions=(Language="ja", Text="最小限の機能とコンテンツが備わった空の テンプレートです。\n\nこのテンプレートを利用すると、Datasmith を使ってモデルをさまざまなデザイン アプリケーションからインポートできるようになるとともに、それらを簡単にライティングできるようになります。")
 
 
 ; The sort key to use for the template to bypass the default alphabetical sorting.
diff --git a/Templates/TP_ME_BlankBP/Config/TemplateDefs.ini b/Templates/TP_ME_BlankBP/Config/TemplateDefs.ini
--- a/Templates/TP_ME_BlankBP/Config/TemplateDefs.ini
+++ b/Templates/TP_ME_BlankBP/Config/TemplateDefs.ini
@@ -19,7 +19,8 @@ LocalizedDescriptions=(Language="ko", Text="최소한의 기능을 가진 빈 
 LocalizedDisplayNames=(Language="zh-Hans", Text="空白")
 LocalizedDescriptions=(Language="zh-Hans", Text="具备基本功能的空白模板。")
 
-LocalizedDescriptions=(Language="ja", Text="最小限の機能が備わった Blank テンプレートです。")
+LocalizedDisplayNames=(Language="ja", Text="ブランク")
+LocalizedDescriptions=(Language="ja", Text="最小限の機能が備わった 空のテンプレートです。")
 
 ; The sort key to use for the template to bypass the default alphabetical sorting.
 SortKey="_1"
diff --git a/Templates/TP_ME_VProdBP/Config/TemplateDefs.ini b/Templates/TP_ME_VProdBP/Config/TemplateDefs.ini
--- a/Templates/TP_ME_VProdBP/Config/TemplateDefs.ini
+++ b/Templates/TP_ME_VProdBP/Config/TemplateDefs.ini
@@ -19,7 +19,8 @@ LocalizedDescriptions=(Language="ko", Text="이 버추얼 프로덕션 프로젝
 LocalizedDisplayNames=(Language="zh-Hans", Text="虚拟制片")
 LocalizedDescriptions=(Language="zh-Hans", Text="该虚拟制片项目模板包含VR探查、虚拟摄像机、SDI视频、现场合成和nDisplay功能。")
 
-LocalizedDescriptions=(Language="ja", Text="この Virtual Production プロジェクト テンプレートには、VR Scouting、Virtual Camera、SDI Video、Live Compositing、nDisplay のための機能が含まれています。")
+LocalizedDisplayNames=(Language="ja", Text="バーチャルプロダクション")
+LocalizedDescriptions=(Language="ja", Text="この バーチャルプロダクション プロジェクト テンプレートには、VRスカウティング、バーチャルカメラ、SDIビデオ、ライブコンポジット、nDisplay のための機能が含まれています。")
 
 ; The sort key to use for the template to bypass the default alphabetical sorting.
 SortKey="_2"
diff --git a/Templates/TP_MFG_ProdConfigBP/Config/TemplateDefs.ini b/Templates/TP_MFG_ProdConfigBP/Config/TemplateDefs.ini
--- a/Templates/TP_MFG_ProdConfigBP/Config/TemplateDefs.ini
+++ b/Templates/TP_MFG_ProdConfigBP/Config/TemplateDefs.ini
@@ -11,6 +11,7 @@ Categories=MFG
 LocalizedDisplayNames=(Language="en", Text="Product Configurator")
 LocalizedDisplayNames=(Language="ko", Text="제품 컨피규레이터")
 LocalizedDisplayNames=(Language="zh-Hans", Text="产品配置器")
+LocalizedDisplayNames=(Language="ja", Text="製品コンフィギュレーター")
 
 ; The localized descriptions for this template, which is displayed when the template is selected.
 LocalizedDescriptions=(Language="en", Text="This project template uses the Variant Manager, UMG and Blueprint to build a generic product configurator.\n\nThe user interface is dynamically created from the 3d model and the variants present in the level. This allows for users to swap the provided model for user-made assets and have the configurator's UI adapting to the new data.")
diff --git a/Templates/TP_MFG_ProdConfigBP/manifest.json b/Templates/TP_MFG_ProdConfigBP/manifest.json
--- a/Templates/TP_MFG_ProdConfigBP/manifest.json
+++ b/Templates/TP_MFG_ProdConfigBP/manifest.json
@@ -6,6 +6,10 @@
 			"Language": "en",
 			"Text": "MFG Product Configurator"
 		},
+		{
+			"Language": "ja",
+			"Text": "製品コンフィギュレーター"
+		},
 		{
 			"Language": "ko",
 			"Text": "제품 컨피규레이터"
diff --git a/Templates/TP_PhotoStudioBP/Config/TemplateDefs.ini b/Templates/TP_PhotoStudioBP/Config/TemplateDefs.ini
--- a/Templates/TP_PhotoStudioBP/Config/TemplateDefs.ini
+++ b/Templates/TP_PhotoStudioBP/Config/TemplateDefs.ini
@@ -19,6 +19,7 @@ LocalizedDescriptions=(Language="ko", Text="이 템플릿은 사진 스튜디오
 LocalizedDisplayNames=(Language="zh-Hans", Text="摄影棚")
 LocalizedDescriptions=(Language="zh-Hans", Text="该模板可以让你在摄影棚环境中创建产品演示和动画。\n\n该模板允许使用Datasmith从多种设计软件中导入模型并轻松为其设置光照。")
 
+LocalizedDisplayNames=(Language="ja", Text="フォトスタジオ")
 LocalizedDescriptions=(Language="ja", Text="このテンプレートでは、製品のプレゼンテーションおよびシネマティックスをフォトスタジオ環境で制作できます。\n\nこのテンプレートを利用すると、Datasmith を使ってモデルをさまざまなデザイン アプリケーションからインポートできるようになるとともに、それらを簡単にライティングできるようになります。")
 
 ; The sort key to use for the template to bypass the default alphabetical sorting.
diff --git a/Templates/TP_SIM_Blank/Config/TemplateDefs.ini b/Templates/TP_SIM_Blank/Config/TemplateDefs.ini
--- a/Templates/TP_SIM_Blank/Config/TemplateDefs.ini
+++ b/Templates/TP_SIM_Blank/Config/TemplateDefs.ini
@@ -11,7 +11,7 @@ Categories=SIM
 LocalizedDisplayNames=(Language="en", Text="Simulation Blank")
 LocalizedDisplayNames=(Language="zh-Hans", Text="模拟-空白")
 LocalizedDisplayNames=(Language="ko", Text="시뮬레이션 기본")
-LocalizedDisplayNames=(Language="ja", Text="シミュレーション Blank")
+LocalizedDisplayNames=(Language="ja", Text="シミュレーション ブランク")
 
 ; The localized descriptions for this template, which is displayed when the template is selected.
 LocalizedDescriptions=(Language="en", Text="A simple empty project template enabling initial settings and plugins. A starting point suitable for most of your simulation applications.\n\nEarth Atmosphere, Lighting, Volumetric Clouds, Geo Coordinates, WGS84")
diff --git a/Templates/TP_SIM_Blank/manifest.json b/Templates/TP_SIM_Blank/manifest.json
--- a/Templates/TP_SIM_Blank/manifest.json
+++ b/Templates/TP_SIM_Blank/manifest.json
@@ -16,7 +16,7 @@
 		},
 		{
 			"Language": "ja",
-			"Text": "シミュレーション Blank"
+			"Text": "シミュレーション ブランク"
 		}
 	],
 	"Description":
